Fix property access nodes creating and using an incorrect internal property type

Previously there was no checking to see whether a variable that was created actually exists already, which could cause collisions.
Also fixed issue where connecting to a wildcard pin would mean that the output pin always stayed a wildcard.

diff --git a/Engine/Plugins/Developer/PropertyAccessNode/Source/PropertyAccessNode/Private/K2Node_PropertyAccess.cpp b/Engine/Plugins/Developer/PropertyAccessNode/Source/PropertyAccessNode/Private/K2Node_PropertyAccess.cpp
--- a/Engine/Plugins/Developer/PropertyAccessNode/Source/PropertyAccessNode/Private/K2Node_PropertyAccess.cpp
+++ b/Engine/Plugins/Developer/PropertyAccessNode/Source/PropertyAccessNode/Private/K2Node_PropertyAccess.cpp
@@ -31,7 +31,7 @@ void UK2Node_PropertyAccess::CreateClassVariablesFromBlueprint(IAnimBlueprintVar
 	if(ResolvedPinType != FEdGraphPinType() && ResolvedPinType.PinCategory != UEdGraphSchema_K2::PC_Wildcard)
 	{
 		// Create internal generated destination property
-		if(FProperty* DestProperty = InCreationContext.CreateVariable(GetFName(), ResolvedPinType))
+		if(FProperty* DestProperty = InCreationContext.CreateUniqueVariable(this, ResolvedPinType))
 		{
 			GeneratedPropertyName = DestProperty->GetFName();
 		}
@@ -136,7 +136,7 @@ void UK2Node_PropertyAccess::AllocatePins(UEdGraphPin* InOldOutputPin)
 	{
 		UEdGraphPin* OutputPin = nullptr;
 	
-		if(InOldOutputPin != nullptr && InOldOutputPin->LinkedTo.Num() > 0)
+		if(InOldOutputPin != nullptr && InOldOutputPin->LinkedTo.Num() > 0 && InOldOutputPin->PinType.PinCategory != UEdGraphSchema_K2::PC_Wildcard)
 		{
 			// Use old output pin if we have one and it is connected
 			OutputPin = CreatePin(EGPD_Output, InOldOutputPin->PinType, TEXT("Value"));
diff --git a/Engine/Source/Editor/AnimGraph/Private/AnimBlueprintCompiler.cpp b/Engine/Source/Editor/AnimGraph/Private/AnimBlueprintCompiler.cpp
--- a/Engine/Source/Editor/AnimGraph/Private/AnimBlueprintCompiler.cpp
+++ b/Engine/Source/Editor/AnimGraph/Private/AnimBlueprintCompiler.cpp
@@ -1249,6 +1249,12 @@ TArrayView<const FName> FAnimBlueprintCompilerContext::GetAttributesFromNode(UAn
 	return TArrayView<const FName>();
 }
 
+FProperty* FAnimBlueprintCompilerContext::CreateUniqueVariable(UObject* InForObject, const FEdGraphPinType& Type)
+{
+	const FString VariableName = ClassScopeNetNameMap.MakeValidName(InForObject);
+	return CreateVariable(*VariableName, Type);
+}
+
 //////////////////////////////////////////////////////////////////////////
 
 #undef LOCTEXT_NAMESPACE
\ No newline at end of file
diff --git a/Engine/Source/Editor/AnimGraph/Private/AnimBlueprintCompiler.h b/Engine/Source/Editor/AnimGraph/Private/AnimBlueprintCompiler.h
--- a/Engine/Source/Editor/AnimGraph/Private/AnimBlueprintCompiler.h
+++ b/Engine/Source/Editor/AnimGraph/Private/AnimBlueprintCompiler.h
@@ -175,5 +175,8 @@ private:
 
 	// Get the current compiled-in attributes uniquely assigned to the specified node
 	TArrayView<const FName> GetAttributesFromNode(UAnimGraphNode_Base* InNode) const;
+	
+	// Create a uniquely named variable corresponding to an object in the current class
+	FProperty* CreateUniqueVariable(UObject* InForObject, const FEdGraphPinType& Type);
 };
 
diff --git a/Engine/Source/Editor/AnimGraph/Private/AnimBlueprintVariableCreationContext.cpp b/Engine/Source/Editor/AnimGraph/Private/AnimBlueprintVariableCreationContext.cpp
--- a/Engine/Source/Editor/AnimGraph/Private/AnimBlueprintVariableCreationContext.cpp
+++ b/Engine/Source/Editor/AnimGraph/Private/AnimBlueprintVariableCreationContext.cpp
@@ -7,3 +7,8 @@ FProperty* FAnimBlueprintVariableCreationContext::CreateVariable(const FName Nam
 {
 	return CompilerContext->CreateVariable(Name, Type);
 }
+
+FProperty* FAnimBlueprintVariableCreationContext::CreateUniqueVariable(UObject* InForObject, const FEdGraphPinType& Type)
+{
+	return CompilerContext->CreateUniqueVariable(InForObject, Type);
+}
diff --git a/Engine/Source/Editor/AnimGraph/Private/AnimBlueprintVariableCreationContext.h b/Engine/Source/Editor/AnimGraph/Private/AnimBlueprintVariableCreationContext.h
--- a/Engine/Source/Editor/AnimGraph/Private/AnimBlueprintVariableCreationContext.h
+++ b/Engine/Source/Editor/AnimGraph/Private/AnimBlueprintVariableCreationContext.h
@@ -17,6 +17,7 @@ private:
 	
 	// IAnimBlueprintVariableCreationContext interface
 	virtual FProperty* CreateVariable(const FName Name, const FEdGraphPinType& Type) override;
+	virtual FProperty* CreateUniqueVariable(UObject* InForObject, const FEdGraphPinType& Type) override;
 
 	FAnimBlueprintCompilerContext* CompilerContext;
 };
\ No newline at end of file
diff --git a/Engine/Source/Editor/AnimGraph/Public/IClassVariableCreator.h b/Engine/Source/Editor/AnimGraph/Public/IClassVariableCreator.h
--- a/Engine/Source/Editor/AnimGraph/Public/IClassVariableCreator.h
+++ b/Engine/Source/Editor/AnimGraph/Public/IClassVariableCreator.h
@@ -15,8 +15,11 @@ struct FEdGraphPinType;
 class IAnimBlueprintVariableCreationContext
 {
 public:
-	/** Create a class variable in the current class */
+	/** Create a class variable in the current class. Note that no name confick resolution is performed, if a unique name is needed, use CreateUniqueVariable */
 	virtual FProperty* CreateVariable(const FName Name, const FEdGraphPinType& Type) = 0;
+	
+	/** Create a uniquely named variable corresponding to an object in the current class. */
+	virtual FProperty* CreateUniqueVariable(UObject* InForObject, const FEdGraphPinType& Type) = 0;
 };
 
 UINTERFACE(MinimalAPI)
