Add missing header guards

diff --git a/Engine/Plugins/Developer/PropertyAccessNode/Source/PropertyAccessNode/Private/PropertyAccessNodeLibrary.h b/Engine/Plugins/Developer/PropertyAccessNode/Source/PropertyAccessNode/Private/PropertyAccessNodeLibrary.h
--- a/Engine/Plugins/Developer/PropertyAccessNode/Source/PropertyAccessNode/Private/PropertyAccessNodeLibrary.h
+++ b/Engine/Plugins/Developer/PropertyAccessNode/Source/PropertyAccessNode/Private/PropertyAccessNodeLibrary.h
@@ -1,3 +1,4 @@
 // Copyright Epic Games, Inc. All Rights Reserved.
+#pragma once
 
 #include "CoreMinimal.h"
diff --git a/Engine/Source/Runtime/AnimGraphRuntime/Public/AnimNodes/AnimNode_CallFunction.h b/Engine/Source/Runtime/AnimGraphRuntime/Public/AnimNodes/AnimNode_CallFunction.h
--- a/Engine/Source/Runtime/AnimGraphRuntime/Public/AnimNodes/AnimNode_CallFunction.h
+++ b/Engine/Source/Runtime/AnimGraphRuntime/Public/AnimNodes/AnimNode_CallFunction.h
@@ -1,4 +1,5 @@
 ﻿// Copyright Epic Games, Inc. All Rights Reserved.
+#pragma once
 
 #include "CoreMinimal.h"
 #include "Animation/AnimNodeBase.h"
