Fixed property access nodes reporting an error when using split pins

Calling Super::ExpandNode applies split pin logic

diff --git a/Engine/Plugins/Developer/PropertyAccessNode/Source/PropertyAccessNode/Private/K2Node_PropertyAccess.cpp b/Engine/Plugins/Developer/PropertyAccessNode/Source/PropertyAccessNode/Private/K2Node_PropertyAccess.cpp
--- a/Engine/Plugins/Developer/PropertyAccessNode/Source/PropertyAccessNode/Private/K2Node_PropertyAccess.cpp
+++ b/Engine/Plugins/Developer/PropertyAccessNode/Source/PropertyAccessNode/Private/K2Node_PropertyAccess.cpp
@@ -42,6 +42,8 @@ void UK2Node_PropertyAccess::CreateClassVariablesFromBlueprint(IAnimBlueprintVar
 
 void UK2Node_PropertyAccess::ExpandNode(FKismetCompilerContext& InCompilerContext, UEdGraph* InSourceGraph)
 {
+	Super::ExpandNode(InCompilerContext, InSourceGraph);
+	
 	ResolvePropertyAccess();
 
 	UK2Node_PropertyAccess* OriginalNode = CastChecked<UK2Node_PropertyAccess>(InCompilerContext.MessageLog.FindSourceObject(this));
