Fixed crash when splitting a post process settings pin on a property access node

diff --git a/Engine/Plugins/Developer/PropertyAccessNode/Source/PropertyAccessNode/Private/SPropertyAccessNode.cpp b/Engine/Plugins/Developer/PropertyAccessNode/Source/PropertyAccessNode/Private/SPropertyAccessNode.cpp
--- a/Engine/Plugins/Developer/PropertyAccessNode/Source/PropertyAccessNode/Private/SPropertyAccessNode.cpp
+++ b/Engine/Plugins/Developer/PropertyAccessNode/Source/PropertyAccessNode/Private/SPropertyAccessNode.cpp
@@ -30,9 +30,12 @@ void SPropertyAccessNode::Construct(const FArguments& InArgs, UK2Node_PropertyAc
 	UpdateGraphNode();
 
 	// Centre the pin slot
-	SVerticalBox::FSlot& PinSlot = RightNodeBox->GetSlot(0);
-	PinSlot.SetVerticalAlignment(VAlign_Center);
-	PinSlot.SetFillHeight(1.0f);
+	if(RightNodeBox->GetChildren()->Num() > 0)
+	{
+		SVerticalBox::FSlot& PinSlot = RightNodeBox->GetSlot(0);
+		PinSlot.SetVerticalAlignment(VAlign_Center);
+		PinSlot.SetFillHeight(1.0f);
+	}
 }
 
 bool SPropertyAccessNode::CanBindProperty(FProperty* InProperty) const
