GitHub #10551: Fix `PropertyAccess` Crash when used in child AnimBP

Slightly modified from the original PR to add a reasoning comment and fallback error in case we are somehow not in a child anim BP context

diff --git a/Engine/Plugins/Developer/PropertyAccessNode/Source/PropertyAccessNode/Private/K2Node_PropertyAccess.cpp b/Engine/Plugins/Developer/PropertyAccessNode/Source/PropertyAccessNode/Private/K2Node_PropertyAccess.cpp
--- a/Engine/Plugins/Developer/PropertyAccessNode/Source/PropertyAccessNode/Private/K2Node_PropertyAccess.cpp
+++ b/Engine/Plugins/Developer/PropertyAccessNode/Source/PropertyAccessNode/Private/K2Node_PropertyAccess.cpp
@@ -46,8 +46,25 @@ void UK2Node_PropertyAccess::ExpandNode(FKismetCompilerContext& InCompilerContex
 	OriginalNode->CompiledContextDesc = FText();
 
 	TUniquePtr<IAnimBlueprintCompilationContext> CompilationContext = IAnimBlueprintCompilationContext::Get(InCompilerContext);
-	UAnimBlueprintExtension_PropertyAccess* PropertyAccessExtension = UAnimBlueprintExtension::GetExtension<UAnimBlueprintExtension_PropertyAccess>(CastChecked<UAnimBlueprint>(GetBlueprint()));
-	check(PropertyAccessExtension);
+	UAnimBlueprint* AnimBlueprint = CastChecked<UAnimBlueprint>(GetBlueprint());
+	UAnimBlueprintExtension_PropertyAccess* PropertyAccessExtension = UAnimBlueprintExtension::FindExtension<UAnimBlueprintExtension_PropertyAccess>(AnimBlueprint);
+
+	if (!PropertyAccessExtension)
+	{
+		// Child anim blueprints cannot modify or augment the property access library of their parent classes, so any additional 
+		// property access nodes in child classes will not work at present. To address this (and likely a number of other missing 
+		// child anim BP features) we would need to allow child anim BPs to perform a 'full' compilation, so they can add additional
+		// anim graph nodes & property access libraries, amongst other things.
+		if (UAnimBlueprint* ParentAnimBP = UAnimBlueprint::GetParentAnimBlueprint(AnimBlueprint))
+		{
+			InCompilerContext.MessageLog.Error(*LOCTEXT("PropertyAccess_AnimBlueprintExtension_PropertyAccess_NotFound_ChildAnimBP", "@@ is currently unsupported in Child Animation Blueprints.").ToString(), this);
+		}
+		else
+		{
+			InCompilerContext.MessageLog.Error(*LOCTEXT("PropertyAccess_AnimBlueprintExtension_PropertyAccess_NotFound", "Internal compiler error: Property access extension not found for @@.").ToString(), this);
+		}
+		return;
+	}
 	
 	if(GeneratedPropertyName != NAME_None)
 	{
