Fix Lyra PIE property access issues caused by CL 31251549

The changes to always process 'most up to date' classes broke property accesses on many of Lyra's accesses. This was because the runtime functions/properties ended up being accessed on empty skeleton classes.
This change restricts the path resolution via 'most up to date' classes to only where it is needed for this use case: renaming/referencing variables and functions.

diff --git a/Engine/Plugins/Developer/PropertyAccessNode/Source/PropertyAccessNode/Private/K2Node_PropertyAccess.cpp b/Engine/Plugins/Developer/PropertyAccessNode/Source/PropertyAccessNode/Private/K2Node_PropertyAccess.cpp
--- a/Engine/Plugins/Developer/PropertyAccessNode/Source/PropertyAccessNode/Private/K2Node_PropertyAccess.cpp
+++ b/Engine/Plugins/Developer/PropertyAccessNode/Source/PropertyAccessNode/Private/K2Node_PropertyAccess.cpp
@@ -321,6 +321,7 @@ void UK2Node_PropertyAccess::HandleVariableRenamed(UBlueprint* InBlueprint, UCla
 	// See if the path references the variable
 	TArray<int32> RenameIndices;
 	IPropertyAccessEditor::FResolvePropertyAccessArgs ResolveArgs;
+	ResolveArgs.bUseMostUpToDateClasses = true;
 	ResolveArgs.PropertyFunction = [InOldVarName, SkeletonVariableClass, &RenameIndices](int32 InSegmentIndex, FProperty* InProperty, int32 InStaticArrayIndex)
 	{
 		UClass* OwnerClass = InProperty->GetOwnerClass();
@@ -349,6 +350,7 @@ void UK2Node_PropertyAccess::HandleFunctionRenamed(UBlueprint* InBlueprint, UCla
 	// See if the path references the variable
 	TArray<int32> RenameIndices;
 	IPropertyAccessEditor::FResolvePropertyAccessArgs ResolveArgs;
+	ResolveArgs.bUseMostUpToDateClasses = true;
 	ResolveArgs.FunctionFunction = [InOldFuncName, SkeletonFunctionClass, &RenameIndices](int32 InSegmentIndex, UFunction* InFunction, FProperty* InReturnProperty)
 	{
 		const UClass* OwnerSkeletonFunctionClass = FBlueprintEditorUtils::GetMostUpToDateClass(InFunction->GetOuterUClass());
@@ -382,6 +384,7 @@ void UK2Node_PropertyAccess::ReplaceReferences(UBlueprint* InBlueprint, UBluepri
 	// See if the path references the variable
 	TArray<int32> ReplaceIndices;
 	IPropertyAccessEditor::FResolvePropertyAccessArgs ResolveArgs;
+	ResolveArgs.bUseMostUpToDateClasses = true;
 	ResolveArgs.PropertyFunction = [SourceProperty, &ReplaceIndices](int32 InSegmentIndex, FProperty* InProperty, int32 InStaticArrayIndex)
 	{
 		if(InProperty == SourceProperty)
@@ -410,6 +413,7 @@ bool UK2Node_PropertyAccess::ReferencesVariable(const FName& InVarName, const US
 	bool bReferencesVariable = false;
 	
 	IPropertyAccessEditor::FResolvePropertyAccessArgs ResolveArgs;
+	ResolveArgs.bUseMostUpToDateClasses = true;
 	ResolveArgs.PropertyFunction = [InVarName, SkeletonVariableClass, &bReferencesVariable](int32 InSegmentIndex, FProperty* InProperty, int32 InStaticArrayIndex)
 	{
 		if(SkeletonVariableClass)
@@ -442,6 +446,7 @@ bool UK2Node_PropertyAccess::ReferencesFunction(const FName& InFunctionName, con
 	bool bReferencesFunction = false;
 
 	IPropertyAccessEditor::FResolvePropertyAccessArgs ResolveArgs;
+	ResolveArgs.bUseMostUpToDateClasses = true;
 	ResolveArgs.FunctionFunction = [InFunctionName, SkeletonFunctionClass, &bReferencesFunction](int32 InSegmentIndex, UFunction* InFunction, FProperty* InProperty)
 	{
 		if (SkeletonFunctionClass)
diff --git a/Engine/Plugins/Runtime/PropertyAccess/Source/PropertyAccessEditor/Private/PropertyAccessEditor.cpp b/Engine/Plugins/Runtime/PropertyAccess/Source/PropertyAccessEditor/Private/PropertyAccessEditor.cpp
--- a/Engine/Plugins/Runtime/PropertyAccess/Source/PropertyAccessEditor/Private/PropertyAccessEditor.cpp
+++ b/Engine/Plugins/Runtime/PropertyAccess/Source/PropertyAccessEditor/Private/PropertyAccessEditor.cpp
@@ -50,6 +50,12 @@ struct FPropertyAccessEditorSystem
 
 		// Whether to apply allow-list to functions and properties along the path
 		bool bPerformValidation = false;
+
+		// Whether to use the most up to date classes when traversing the path. 
+		// This can be useful for situations where we are resolving against potentially out of date
+		// classes, but the resulting path will not be valid to use or persist due to functions and properties
+		// being on skeleton classes 
+		bool bUseMostUpToDateClasses = false;
 	};
 
 	// The result of a segment resolve operation
@@ -246,7 +252,7 @@ struct FPropertyAccessEditorSystem
 
 				InContext.CurrentStruct = Segment.Struct;
 #if WITH_EDITOR
-				if (InContext.CurrentStruct && InContext.CurrentStruct->IsA<UClass>())
+				if (InContext.bUseMostUpToDateClasses && InContext.CurrentStruct && InContext.CurrentStruct->IsA<UClass>())
 				{
 					// Try to use the skeleton class if possible in editor as the regular class may have yet to be compiled
 					InContext.CurrentStruct = FBlueprintEditorUtils::GetMostUpToDateClass(CastChecked<UClass>(InContext.CurrentStruct));
@@ -292,6 +298,7 @@ struct FPropertyAccessEditorSystem
 	{
 		FPropertyAccessPath AccessPath;
 		FResolveSegmentsContext Context(InStruct, InPath, AccessPath);
+		Context.bUseMostUpToDateClasses = InArgs.bUseMostUpToDateClasses;
 		FPropertyAccessResolveResult Result;
 		Result.Result = ResolveSegments(Context);
 		Result.bIsThreadSafe = Context.bWasThreadSafe;
diff --git a/Engine/Source/Editor/AnimGraph/Private/AnimGraphNodeBinding_Base.cpp b/Engine/Source/Editor/AnimGraph/Private/AnimGraphNodeBinding_Base.cpp
--- a/Engine/Source/Editor/AnimGraph/Private/AnimGraphNodeBinding_Base.cpp
+++ b/Engine/Source/Editor/AnimGraph/Private/AnimGraphNodeBinding_Base.cpp
@@ -147,6 +147,7 @@ void UAnimGraphNodeBinding_Base::HandleVariableRenamed(UBlueprint* InBlueprint,
 		{
 			TArray<int32> RenameIndices;
 			IPropertyAccessEditor::FResolvePropertyAccessArgs ResolveArgs;
+			ResolveArgs.bUseMostUpToDateClasses = true;
 			ResolveArgs.PropertyFunction = [InOldVarName, SkeletonVariableClass, &RenameIndices](int32 InSegmentIndex, FProperty* InProperty, int32 InStaticArrayIndex)
 			{
 				UClass* OwnerClass = InProperty->GetOwnerClass();
@@ -183,6 +184,7 @@ void UAnimGraphNodeBinding_Base::HandleFunctionRenamed(UBlueprint* InBlueprint,
 		{
 			TArray<int32> RenameIndices;
 			IPropertyAccessEditor::FResolvePropertyAccessArgs ResolveArgs;
+			ResolveArgs.bUseMostUpToDateClasses = true;
 			ResolveArgs.FunctionFunction = [InOldFuncName, SkeletonFunctionClass, &RenameIndices](int32 InSegmentIndex, UFunction* InFunction, FProperty* InReturnProperty)
 			{
 				const UClass* OwnerSkeletonFunctionClass = FBlueprintEditorUtils::GetMostUpToDateClass(InFunction->GetOuterUClass());
@@ -224,6 +226,7 @@ void UAnimGraphNodeBinding_Base::ReplaceReferences(UBlueprint* InBlueprint, UBlu
 		{
 			TArray<int32> ReplaceIndices;
 			IPropertyAccessEditor::FResolvePropertyAccessArgs ResolveArgs;
+			ResolveArgs.bUseMostUpToDateClasses = true;
 			ResolveArgs.PropertyFunction = [SourceProperty, &ReplaceIndices](int32 InSegmentIndex, FProperty* InProperty, int32 InStaticArrayIndex)
 			{
 				if (InProperty == SourceProperty)
@@ -260,6 +263,7 @@ bool UAnimGraphNodeBinding_Base::ReferencesVariable(const FName& InVarName, cons
 			bool bReferencesVariable = false;
 
 			IPropertyAccessEditor::FResolvePropertyAccessArgs ResolveArgs;
+			ResolveArgs.bUseMostUpToDateClasses = true;
 			ResolveArgs.PropertyFunction = [InVarName, SkeletonVariableClass, &bReferencesVariable](int32 InSegmentIndex, FProperty* InProperty, int32 InStaticArrayIndex)
 			{
 				if (SkeletonVariableClass)
@@ -305,6 +309,7 @@ bool UAnimGraphNodeBinding_Base::ReferencesFunction(const FName& InFunctionName,
 			bool bReferencesFunction = false;
 
 			IPropertyAccessEditor::FResolvePropertyAccessArgs ResolveArgs;
+			ResolveArgs.bUseMostUpToDateClasses = true;
 			ResolveArgs.FunctionFunction = [InFunctionName, SkeletonFunctionClass, &bReferencesFunction](int32 InSegmentIndex, UFunction* InFunction, FProperty* InProperty)
 			{
 				if (SkeletonFunctionClass)
diff --git a/Engine/Source/Editor/UnrealEd/Public/IPropertyAccessEditor.h b/Engine/Source/Editor/UnrealEd/Public/IPropertyAccessEditor.h
--- a/Engine/Source/Editor/UnrealEd/Public/IPropertyAccessEditor.h
+++ b/Engine/Source/Editor/UnrealEd/Public/IPropertyAccessEditor.h
@@ -385,6 +385,14 @@ public:
 
 		/** Function called when a function is resolved. */
 		TFunction<void(int32 /*SegmentIndex*/, UFunction* /*Function*/, FProperty* /*ReturnProperty*/)> FunctionFunction;
+
+		/** 
+		 * Whether to use the most up to date classes when traversing the path. 
+		 * This can be useful for situations where we are resolving against potentially out of date
+		 * classes, but the resulting path will not be valid to use or persist due to functions and properties
+		 * being on skeleton classes 
+		 */
+		bool bUseMostUpToDateClasses = false;
 	};
 	
 	/**
